modified the IsValidOrFixable() function for RawMesh as it was more conservative than the code after it - it crashed later where the method returned true
 - if that code triggers we need to fix the input of the code using the data

diff --git a/Engine/Source/Developer/RawMesh/Private/RawMesh.cpp b/Engine/Source/Developer/RawMesh/Private/RawMesh.cpp
--- a/Engine/Source/Developer/RawMesh/Private/RawMesh.cpp
+++ b/Engine/Source/Developer/RawMesh/Private/RawMesh.cpp
@@ -68,16 +68,19 @@ bool FRawMesh::IsValidOrFixable() const
 	int32 NumVertices = VertexPositions.Num();
 	int32 NumWedges = WedgeIndices.Num();
 	int32 NumFaces = NumWedges / 3;
+	int32 NumTexCoords = WedgeTexCoords[0].Num();
+	int32 NumFaceSmoothingMasks = FaceSmoothingMasks.Num();
+	int32 NumFaceMaterialIndices = FaceMaterialIndices.Num();
 
 	bool bValidOrFixable = NumVertices > 0
 		&& NumWedges > 0
 		&& NumFaces > 0
 		&& (NumWedges / 3) == NumFaces
-		&& ValidateArraySize(FaceMaterialIndices, NumFaces)
-		&& ValidateArraySize(FaceSmoothingMasks, NumFaces)
+		&& NumFaceMaterialIndices == NumFaces
+		&& NumFaceSmoothingMasks == NumFaces
 		&& ValidateArraySize(WedgeColors, NumWedges)
 		// All meshes must have a valid texture coordinate.
-		&& WedgeTexCoords[0].Num() == NumWedges; 
+		&& NumTexCoords == NumWedges; 
 
 	for (int32 TexCoordIndex = 1; TexCoordIndex < MAX_MESH_TEXTURE_COORDS; ++TexCoordIndex)
 	{
