added comments

diff --git a/Engine/Source/Developer/RawMesh/Public/RawMesh.h b/Engine/Source/Developer/RawMesh/Public/RawMesh.h
--- a/Engine/Source/Developer/RawMesh/Public/RawMesh.h
+++ b/Engine/Source/Developer/RawMesh/Public/RawMesh.h
@@ -25,25 +25,25 @@ enum
  */
 struct FRawMesh
 {
-	/** Material index. */
+	/** Material index. Array[FaceId] = int32 */
 	TArray<int32> FaceMaterialIndices;
-	/** Smoothing mask. */
+	/** Smoothing mask. Array[FaceId] = uint32 */
 	TArray<uint32> FaceSmoothingMasks;
 
-	/** Position in local space. */
+	/** Position in local space. Array[VertexId] = float3(x,y,z) */
 	TArray<FVector> VertexPositions;
 
-	/** Index of the vertex at this wedge. */
+	/** Index of the vertex at this wedge. Array[WedgeId] = VertexId */
 	TArray<uint32> WedgeIndices;
-	/** Tangent, U direction. */
+	/** Tangent, U direction. Array[WedgeId] = float3(x,y,z) */
 	TArray<FVector>	WedgeTangentX;
-	/** Tangent, V direction. */
+	/** Tangent, V direction. Array[WedgeId] = float3(x,y,z) */
 	TArray<FVector>	WedgeTangentY;
-	/** Normal. */
+	/** Normal. Array[WedgeId] = float3(x,y,z) */
 	TArray<FVector>	WedgeTangentZ;
-	/** Texture coordinates. */
+	/** Texture coordinates. Array[UVId][WedgeId]=float2(u,v) */
 	TArray<FVector2D> WedgeTexCoords[MAX_MESH_TEXTURE_COORDS];
-	/** Color. */
+	/** Color. Array[WedgeId]=float3(r,g,b,a) */
 	TArray<FColor> WedgeColors;
 
 	/**
